[Input] Add a thread affinitiy for input for IInputDevice so that we can specify which input modules are safe to be polled from outside the main game thread.

In the future where we have input processing on a sepeate thread, there may be some input events which can only run on the main thread. For example, windows messages will only fire on the same thread which the window handle is created from. So, we need a way to make sure we only poll events on threads which can safely be called from outside the main thread.

For now, based on the docs, that includes:
- XInput
- Game Input
- WinDualShock.

Right now, this will do nothing, because we have not yet added the game thread.

#rnx
#Synced-CL 52309720

diff --git a/Engine/Plugins/Runtime/GameInput/Source/GameInputBase/Public/IGameInputDeviceInterface.h b/Engine/Plugins/Runtime/GameInput/Source/GameInputBase/Public/IGameInputDeviceInterface.h
--- a/Engine/Plugins/Runtime/GameInput/Source/GameInputBase/Public/IGameInputDeviceInterface.h
+++ b/Engine/Plugins/Runtime/GameInput/Source/GameInputBase/Public/IGameInputDeviceInterface.h
@@ -32,6 +32,8 @@ public:
 protected:
 
 	//~ Begin IInputDevice interface
+	// GameInput polling is designed for multi-threaded access per Microsoft docs.
+	virtual EInputDeviceThreadAffinity GetThreadAffinity() const override { return EInputDeviceThreadAffinity::InputThreadSafe; }
 	GAMEINPUTBASE_API virtual void Tick(float DeltaTime) override;
 	GAMEINPUTBASE_API virtual void SendControllerEvents() override;
 	GAMEINPUTBASE_API virtual void SetMessageHandler(const TSharedRef<FGenericApplicationMessageHandler>& InMessageHandler) override;
diff --git a/Engine/Plugins/Runtime/Windows/WinDualShock/Source/WinDualShock/Private/WinDualShock.cpp b/Engine/Plugins/Runtime/Windows/WinDualShock/Source/WinDualShock/Private/WinDualShock.cpp
--- a/Engine/Plugins/Runtime/Windows/WinDualShock/Source/WinDualShock/Private/WinDualShock.cpp
+++ b/Engine/Plugins/Runtime/Windows/WinDualShock/Source/WinDualShock/Private/WinDualShock.cpp
@@ -117,6 +117,8 @@ public:
 	{
 	}
 
+	virtual EInputDeviceThreadAffinity GetThreadAffinity() const override { return EInputDeviceThreadAffinity::InputThreadSafe; }
+
 	virtual void Tick(float DeltaTime) override
 	{
 	}
diff --git a/Engine/Plugins/Runtime/Windows/XInputDevice/Source/XInputDevice/Private/XInputInterface.h b/Engine/Plugins/Runtime/Windows/XInputDevice/Source/XInputDevice/Private/XInputInterface.h
--- a/Engine/Plugins/Runtime/Windows/XInputDevice/Source/XInputDevice/Private/XInputInterface.h
+++ b/Engine/Plugins/Runtime/Windows/XInputDevice/Source/XInputDevice/Private/XInputInterface.h
@@ -70,6 +70,9 @@ public:
 	virtual void Tick( float DeltaTime ) override {};
 	virtual bool Exec( UWorld* InWorld, const TCHAR* Cmd, FOutputDevice& Ar ) { return false; }
 
+	// XInputGetState() is a stateless poll with no game-thread restrictions per Microsoft docs.
+	virtual EInputDeviceThreadAffinity GetThreadAffinity() const override { return EInputDeviceThreadAffinity::InputThreadSafe; }
+
 private:
 
 	XInputInterface( const TSharedRef< FGenericApplicationMessageHandler >& MessageHandler, bool bShouldBePrimaryDevice );
diff --git a/Engine/Source/Runtime/InputDevice/Public/IInputDevice.h b/Engine/Source/Runtime/InputDevice/Public/IInputDevice.h
--- a/Engine/Source/Runtime/InputDevice/Public/IInputDevice.h
+++ b/Engine/Source/Runtime/InputDevice/Public/IInputDevice.h
@@ -9,6 +9,22 @@ struct FForceFeedbackValues;
 enum class FForceFeedbackChannelType;
 struct FInputDeviceProperty;
 
+/**
+ * Governs which thread is allowed to call Tick() and SendControllerEvents() on an IInputDevice.
+ * Output methods (SetChannelValue, SetChannelValues, SetLightColor, SetDeviceProperty, haptics, etc.)
+ * are always called from the game thread regardless of this affinity declaration.
+ *
+ * Default is GameThreadOnly — existing third-party IInputDevice plugins require no changes.
+ */
+enum class EInputDeviceThreadAffinity : uint8
+{
+	/** Device must be ticked by the game thread. Safe for all legacy and Win32 message-pump devices. */
+	GameThreadOnly,
+
+	/** Device may be ticked by the dedicated input thread. Tick() and SendControllerEvents() must be thread-safe. */
+	InputThreadSafe,
+};
+
 /**
  * Input device interface.
  * Useful for plugins/modules to support custom external input devices.
@@ -24,6 +40,14 @@ public:
 	/** Poll for controller state and send events if needed */
 	virtual void SendControllerEvents() = 0;
 
+	/**
+	 * Returns the thread affinity for this input device — which thread is permitted to call
+	 * Tick() and SendControllerEvents(). Defaults to GameThreadOnly so existing devices
+	 * require no changes. Override and return InputThreadSafe only if both methods are
+	 * fully thread-safe (no game-thread-only APIs, no UObjects, no Slate calls).
+	 */
+	virtual EInputDeviceThreadAffinity GetThreadAffinity() const { return EInputDeviceThreadAffinity::GameThreadOnly; }
+
 	/** Set which MessageHandler will get the events from SendControllerEvents. */
 	virtual void SetMessageHandler( const TSharedRef< FGenericApplicationMessageHandler >& InMessageHandler ) = 0;
 
